NullSourceCodeAccessor: add an explicit check for clang-3.5.

- UBT will pick clang++-3.5 even if clang++ is not available, so creating C++ projects should be allowed in this case (UE-14407).

diff --git a/Engine/Plugins/Developer/NullSourceCodeAccess/Source/NullSourceCodeAccess/Private/NullSourceCodeAccessor.cpp b/Engine/Plugins/Developer/NullSourceCodeAccess/Source/NullSourceCodeAccess/Private/NullSourceCodeAccessor.cpp
--- a/Engine/Plugins/Developer/NullSourceCodeAccess/Source/NullSourceCodeAccess/Private/NullSourceCodeAccessor.cpp
+++ b/Engine/Plugins/Developer/NullSourceCodeAccess/Source/NullSourceCodeAccess/Private/NullSourceCodeAccessor.cpp
@@ -7,7 +7,8 @@
 #define LOCTEXT_NAMESPACE "NullSourceCodeAccessor"
 bool FNullSourceCodeAccessor::CanAccessSourceCode() const
 {
-	return FPaths::FileExists(TEXT("/usr/bin/clang++"));
+	// only check the binaries that UBT will/can use. This file must be in sync with LinuxToolChain.cs
+	return FPaths::FileExists(TEXT("/usr/bin/clang++")) || FPaths::FileExists(TEXT("/usr/bin/clang++-3.5"));
 }
 
 FName FNullSourceCodeAccessor::GetFName() const
